Add support for UGS putting the build changelist only in Build.version for modular builds. Monolithic builds read the define from a local-only file called MonolithicVersion.h.

diff --git a/Engine/Programs/UnrealGameSync/UnrealGameSync.ini b/Engine/Programs/UnrealGameSync/UnrealGameSync.ini
new file mode 100644
--- /dev/null
+++ b/Engine/Programs/UnrealGameSync/UnrealGameSync.ini
@@ -0,0 +1,2 @@
+[Options]
+UseFastModularVersioning=true
