Merging UGS SDK info settings from Dev-Core.

diff --git a/Engine/Programs/UnrealGameSync/UnrealGameSync.ini b/Engine/Programs/UnrealGameSync/UnrealGameSync.ini
--- a/Engine/Programs/UnrealGameSync/UnrealGameSync.ini
+++ b/Engine/Programs/UnrealGameSync/UnrealGameSync.ini
@@ -9,4 +9,14 @@ UseFastModularVersioning=1
 +ExcludeChanges=^Changed the dynamic query of collection '.*' to '.*'
 +ExcludeChanges=^Changed the parent of collection '.*'
 +ExcludeChanges=^Upgraded collection '.*' \(was version \d+, now version \d+\)
-+ExcludeChanges=^Collection '.*' not modified
\ No newline at end of file
++ExcludeChanges=^Collection '.*' not modified
+
+[Default]
+; Set this to a network share which contains the SDK installers for your site
+SdkInstallerDir= 
+
+; All the required SDKs for the current version of the engine
++SdkInfo=(Category="Android", Description="NDK r21", Browse="$(SdkInstallerDir)\\Android")
++SdkInfo=(Category="Windows", Description="Visual Studio 2017")
++SdkInfo=(Category="Windows", Description="Visual C++ Toolchain 14.13.26128")
++SdkInfo=(Category="Windows", Description="Windows SDK 10.0.16299.0")
