Enable the UGS build health feature.

diff --git a/Engine/Programs/UnrealGameSync/UnrealGameSync.ini b/Engine/Programs/UnrealGameSync/UnrealGameSync.ini
--- a/Engine/Programs/UnrealGameSync/UnrealGameSync.ini
+++ b/Engine/Programs/UnrealGameSync/UnrealGameSync.ini
@@ -2,6 +2,7 @@
 UseFastModularVersioning=1
 UseFastModularVersioningV2=1
 DebugGameHasSeparateExecutable=1
+ShowBuildHealth=1
 +ExcludeChanges=^Rebuild lightmaps\.$
 +ExcludeChanges=^Rebuild lightmaps & texture streaming\.$
 +ExcludeChanges=^Added collection '.*'\s*
