UGS: Make the filter for excluding automated lighting rebuilds more explicit.

#rnx

diff --git a/Engine/Programs/UnrealGameSync/UnrealGameSync.ini b/Engine/Programs/UnrealGameSync/UnrealGameSync.ini
--- a/Engine/Programs/UnrealGameSync/UnrealGameSync.ini
+++ b/Engine/Programs/UnrealGameSync/UnrealGameSync.ini
@@ -1,6 +1,7 @@
 [Options]
 UseFastModularVersioning=1
-+ExcludeChanges=lightmaps
++ExcludeChanges=^Rebuild lightmaps\.$
++ExcludeChanges=^Rebuild lightmaps & texture streaming\.$
 +ExcludeChanges=^Added collection '.*'\s*
 +ExcludeChanges=^Added '.*' to collection '.*'\s*
 +ExcludeChanges=^Added \d+ objects to collection '.*':
