Enable option to build CrashReportClientEditor in UGS.

diff --git a/Engine/Programs/UnrealGameSync/UnrealGameSync.ini b/Engine/Programs/UnrealGameSync/UnrealGameSync.ini
--- a/Engine/Programs/UnrealGameSync/UnrealGameSync.ini
+++ b/Engine/Programs/UnrealGameSync/UnrealGameSync.ini
@@ -1,6 +1,7 @@
 [Options]
 UseFastModularVersioning=1
 UseFastModularVersioningV2=1
+UseCrashReportClientEditor=1
 DebugGameHasSeparateExecutable=1
 +ExcludeChanges=^Rebuild lightmaps\.$
 +ExcludeChanges=^Rebuild lightmaps & texture streaming\.$
