Move UGS sync category definitions into config files.

diff --git a/Engine/Programs/UnrealGameSync/UnrealGameSync.ini b/Engine/Programs/UnrealGameSync/UnrealGameSync.ini
--- a/Engine/Programs/UnrealGameSync/UnrealGameSync.ini
+++ b/Engine/Programs/UnrealGameSync/UnrealGameSync.ini
@@ -15,6 +15,18 @@ DebugGameHasSeparateExecutable=1
 +ExcludeChanges=^Upgraded collection '.*' \(was version \d+, now version \d+\)
 +ExcludeChanges=^Collection '.*' not modified
 
+; Sync Categories
++SyncCategory=(UniqueId="6703e989-d912-451d-93ad-b48de748d282", Name="Content", Paths="*.uasset;*.umap")
++SyncCategory=(UniqueId="6507c2fb-19dd-403a-afa3-bbf898248d5a", Name="Documentation", Paths="/Engine/Documentation/...")
++SyncCategory=(UniqueId="fd7c716e-4bad-43ae-8fae-8748ef9ee44d", Name="Platform Support: Android", Paths="/Engine/Source/ThirdParty/.../Android/...;.../Build/Android/PipelineCaches/...")
++SyncCategory=(UniqueId="176b2eb2-35f7-4e8e-b131-5f1c5f0959af", Name="Platform Support: iOS", Paths="/Engine/Source/ThirdParty/.../IOS/...;.../Build/IOS/PipelineCaches/...")
++SyncCategory=(UniqueId="f44b2d25-cbc0-4a8f-b6b3-e4a8125533dd", Name="Platform Support: Linux", Paths="/Engine/Source/ThirdParty/.../Linux/...")
++SyncCategory=(UniqueId="2af45231-0d75-463b-bf9f-abb3231091bb", Name="Platform Support: Mac", Paths="/Engine/Source/ThirdParty/.../Mac/...;.../Build/Mac/PipelineCaches/...")
++SyncCategory=(UniqueId="3788a0bc-188c-4a0d-950a-d68175f0d110", Name="Platform Support: tvOS", Paths="/Engine/Source/ThirdParty/.../TVOS/...")
++SyncCategory=(UniqueId="1144e719-fcd7-491b-b0fc-8b4c3565bf79", Name="Platform Support: Win32", Paths="/Engine/Source/ThirdParty/.../Win32/...")
++SyncCategory=(UniqueId="5206ccee-9024-4e36-8b89-f5f5a7d288d2", Name="Platform Support: Win64", Paths="/Engine/Source/ThirdParty/.../Win64/...")
++SyncCategory=(UniqueId="cfec942a-bb90-4f0c-accf-238ecaad9430", Name="Source Code", Paths="/Engine/Source/...")
+
 [Default]
 EditorTarget=UnrealEditor
 
