Update UE5 Main UGS Banner

diff --git a/Engine/Programs/UnrealGameSync/UnrealGameSync.ini b/Engine/Programs/UnrealGameSync/UnrealGameSync.ini
--- a/Engine/Programs/UnrealGameSync/UnrealGameSync.ini
+++ b/Engine/Programs/UnrealGameSync/UnrealGameSync.ini
@@ -38,3 +38,8 @@ SdkInstallerDir=
 +SdkInfo=(Category="Windows", Description="Visual Studio 2017")
 +SdkInfo=(Category="Windows", Description="Visual C++ Toolchain 14.13.26128")
 +SdkInfo=(Category="Windows", Description="Windows SDK 10.0.16299.0")
+
+
+[//UE5/Main/UE4Games.uprojectdirs]
+Message=:alert:  Branching for UE5 Early Access will be on Tuesday Feb 02, 2021.
+
