Added UGS message for UE5 Release-5.0 EarlyAccess

diff --git a/Engine/Programs/UnrealGameSync/UnrealGameSync.ini b/Engine/Programs/UnrealGameSync/UnrealGameSync.ini
--- a/Engine/Programs/UnrealGameSync/UnrealGameSync.ini
+++ b/Engine/Programs/UnrealGameSync/UnrealGameSync.ini
@@ -40,6 +40,6 @@ SdkInstallerDir=
 +SdkInfo=(Category="Windows", Description="Windows SDK 10.0.16299.0")
 
 
-[//UE5/Main/UE4Games.uprojectdirs]
-Message=:alert:  Branching for UE5 Early Access will be on Tuesday Feb 02, 2021.
+[//UE5/Release-5.0-EarlyAccess/UE4Games.uprojectdirs]
+Message=:alert:  Branch is OPEN. Use JIRA FixVersion '5.0-Early Access' for tracking.
 
