update to UGS message

diff --git a/Engine/Programs/UnrealGameSync/UnrealGameSync.ini b/Engine/Programs/UnrealGameSync/UnrealGameSync.ini
--- a/Engine/Programs/UnrealGameSync/UnrealGameSync.ini
+++ b/Engine/Programs/UnrealGameSync/UnrealGameSync.ini
@@ -41,5 +41,5 @@ SdkInstallerDir=
 
 
 [//UE5/Release-5.0-EarlyAccess/UE4Games.uprojectdirs]
-Message=:alert:  Branch is OPEN. Use JIRA FixVersion '5.0-Early Access' for tracking.
+Message=:alert:  Branch is in CONTROLLED LOCKDOWN. Use JIRA FixVersion '5.0-Early Access' for tracking.
 
