UGS: Update requirements for building on Windows.

diff --git a/Engine/Programs/UnrealGameSync/UnrealGameSync.ini b/Engine/Programs/UnrealGameSync/UnrealGameSync.ini
--- a/Engine/Programs/UnrealGameSync/UnrealGameSync.ini
+++ b/Engine/Programs/UnrealGameSync/UnrealGameSync.ini
@@ -35,8 +35,7 @@ SdkInstallerDir=
 
 ; All the required SDKs for the current version of the engine
 +SdkInfo=(Category="Android", Description="NDK r21", Browse="$(SdkInstallerDir)\\Android")
-+SdkInfo=(Category="Windows", Description="Visual Studio 2017")
-+SdkInfo=(Category="Windows", Description="Visual C++ Toolchain 14.13.26128")
++SdkInfo=(Category="Windows", Description="Visual Studio 2019 v16.4")
 +SdkInfo=(Category="Windows", Description="Windows SDK 10.0.16299.0")
 +SdkInfo=(Category="Linux", Description="v17_clang-10.0.1-centos7")
 
