Update UGS message for UE5 EA2

diff --git a/Engine/Programs/UnrealGameSync/UnrealGameSync.ini b/Engine/Programs/UnrealGameSync/UnrealGameSync.ini
--- a/Engine/Programs/UnrealGameSync/UnrealGameSync.ini
+++ b/Engine/Programs/UnrealGameSync/UnrealGameSync.ini
@@ -40,5 +40,5 @@ SdkInstallerDir=
 +SdkInfo=(Category="Linux", Description="v17_clang-10.0.1-centos7")
 
 [//UE5/Release-5.0-EarlyAccess/UE4Games.uprojectdirs]
-Message=:alert:  Branch is in HARDCORE LOCKDOWN. Use JIRA FixVersion '5.0-Early Access' for tracking.
+Message=:alert:  Branch is in HARDCORE LOCKDOWN for UE5 EA2 testing. Use JIRA FixVersion '5.0-Early Access 2' for tracking.
 
