Added SDK info Linux to UGS

diff --git a/Engine/Programs/UnrealGameSync/UnrealGameSync.ini b/Engine/Programs/UnrealGameSync/UnrealGameSync.ini
--- a/Engine/Programs/UnrealGameSync/UnrealGameSync.ini
+++ b/Engine/Programs/UnrealGameSync/UnrealGameSync.ini
@@ -38,7 +38,7 @@ SdkInstallerDir=
 +SdkInfo=(Category="Windows", Description="Visual Studio 2017")
 +SdkInfo=(Category="Windows", Description="Visual C++ Toolchain 14.13.26128")
 +SdkInfo=(Category="Windows", Description="Windows SDK 10.0.16299.0")
-
++SdkInfo=(Category="Linux", Description="v17_clang-10.0.1-centos7")
 
 [//UE5/Release-5.0-EarlyAccess/UE4Games.uprojectdirs]
 Message=:alert:  Branch is in HARDCORE LOCKDOWN. Use JIRA FixVersion '5.0-Early Access' for tracking.
