Minor Tweak to Zen Dashboard settings

diff --git a/Engine/Programs/UnrealGameSync/UnrealGameSync.ini b/Engine/Programs/UnrealGameSync/UnrealGameSync.ini
--- a/Engine/Programs/UnrealGameSync/UnrealGameSync.ini
+++ b/Engine/Programs/UnrealGameSync/UnrealGameSync.ini
@@ -15,6 +15,12 @@ DebugGameHasSeparateExecutable=1
 +ExcludeChanges=^Upgraded collection '.*' \(was version \d+, now version \d+\)
 +ExcludeChanges=^Collection '.*' not modified
 
+[Build]
++Step=(UniqueId="89fe8a79-d259-4c7b-bfb4-68f7218b91c2", Description="Compile UnrealInsights", StatusText="Compile UnrealInsights Step...", EstimatedDuration="1", Type="Compile", Target="UnrealInsights", Platform="Win64", Configuration="Development", Arguments="", bNormalSync="False", bScheduledSync="False", bShowAsTool="False", OrderIndex=20)
++Step=(UniqueId="46312669-5069-428d-8d72-c2416c5a322e", Description="Launch UnrealInsights", StatusText="Running UnrealInsights", EstimatedDuration="1", Type="Other", FileName="$(BranchDir)\\Engine\\Binaries\\Win64\\UnrealInsights.exe", Arguments="", bUseLogWindow="False", bNormalSync="False", bScheduledSync="False", bShowAsTool="True", Requires="89fe8a79-d259-4c7b-bfb4-68f7218b91c2", OrderIndex=21 )
++Step=(UniqueId="bb48ca5b-5682-4432-824d-c451336a6523", Description="Compile Zen Dashboard", StatusText="Compile ZenDashboard Step..", EstimatedDuration="1", Type="Compile", Target="ZenDashboard", Platform="Win64", Configuration="Development", Arguments="", UseLogWindow="True", bNormalSync="False", bScheduledSync="False", bShowAsTool="True")
++Step=(UniqueId="586cc0d3-3914-4df9-acb6-2c02cd9d4fc6", Description="Launch Zen Dashboard", StatusText="Launching Zen Dashboard..", EstimatedDuration="1", Type="Other", FileName="$(BranchDir)\\Engine\\Binaries\\Win64\\ZenDashboard.exe", WorkingDir="", Arguments="", bUseLogWindow="False", bNormalSync="False", bShowAsTool="False", bScheduledSync="False", Link="Zen", Requires="bb48ca5b-5682-4432-824d-c451336a6523" )
+
 ; Sync Categories
 +SyncCategory=(UniqueId="6703e989-d912-451d-93ad-b48de748d282", Name="Content", Paths="*.uasset;*.umap")
 +SyncCategory=(UniqueId="6507c2fb-19dd-403a-afa3-bbf898248d5a", Name="Documentation", Paths="/Engine/Documentation/...")
