Add UnrealGameSync.inin SyncCategory for Hololens, the only missing Platform Support filter

#rnx

diff --git a/Engine/Programs/UnrealGameSync/UnrealGameSync.ini b/Engine/Programs/UnrealGameSync/UnrealGameSync.ini
--- a/Engine/Programs/UnrealGameSync/UnrealGameSync.ini
+++ b/Engine/Programs/UnrealGameSync/UnrealGameSync.ini
@@ -20,6 +20,7 @@ DebugGameHasSeparateExecutable=1
 +SyncCategory=(UniqueId="6507c2fb-19dd-403a-afa3-bbf898248d5a", Name="Documentation", Paths="/Engine/Documentation/...")
 +SyncCategory=(UniqueId="fd7c716e-4bad-43ae-8fae-8748ef9ee44d", Name="Platform Support: Android", Paths="/Engine/Source/ThirdParty/.../Android/...;.../Build/Android/PipelineCaches/...")
 +SyncCategory=(UniqueId="176b2eb2-35f7-4e8e-b131-5f1c5f0959af", Name="Platform Support: iOS", Paths="/Engine/Source/ThirdParty/.../IOS/...;.../Build/IOS/PipelineCaches/...")
++SyncCategory=(UniqueId="c60c38c0-3ca8-46eb-b63d-c8f439d3d540", Name="Platform Support: Hololens", Paths="/Engine/Source/ThirdParty/.../Hololens/...")
 +SyncCategory=(UniqueId="f44b2d25-cbc0-4a8f-b6b3-e4a8125533dd", Name="Platform Support: Linux", Paths="/Engine/Source/ThirdParty/.../Linux/...")
 +SyncCategory=(UniqueId="2af45231-0d75-463b-bf9f-abb3231091bb", Name="Platform Support: Mac", Paths="/Engine/Source/ThirdParty/.../Mac/...;.../Build/Mac/PipelineCaches/...")
 +SyncCategory=(UniqueId="3788a0bc-188c-4a0d-950a-d68175f0d110", Name="Platform Support: tvOS", Paths="/Engine/Source/ThirdParty/.../TVOS/...")
