UnrealGameSync.ini: Add Linux sdk path and installer to config

#rnx

diff --git a/Engine/Programs/UnrealGameSync/UnrealGameSync.ini b/Engine/Programs/UnrealGameSync/UnrealGameSync.ini
--- a/Engine/Programs/UnrealGameSync/UnrealGameSync.ini
+++ b/Engine/Programs/UnrealGameSync/UnrealGameSync.ini
@@ -45,7 +45,7 @@ SdkInstallerDir=
 +SdkInfo=(Category="Windows", Description="Visual Studio 2019 v16.11.5")
 +SdkInfo=(Category="Windows", Description="Windows SDK 10.0.18362.0")
 +SdkInfo=(Category="Windows", Description="LLVM Clang 11.1.0")
-+SdkInfo=(Category="Linux", Description="v19_clang-11.0.1-centos7")
++SdkInfo=(Category="Linux", Description="v19_clang-11.0.1-centos7", Browse="$(SdkInstallerDir)\\Linux", Install="$(SdkInstallerDir)\\Linux\\v19_clang-11.0.1-centos7.exe")
 
 [//UE5/Release-5.0-EarlyAccess/UE4Games.uprojectdirs]
 Message=:alert:  Branch is in HARDCORE LOCKDOWN for UE5 EA2 testing. Use JIRA FixVersion '5.0-Early Access 2' for tracking.
