Add SwitchboardListener and UnrealMultiUserServer to the build steps for
UGS.  Default them to off.

diff --git a/Engine/Programs/UnrealGameSync/UnrealGameSync.ini b/Engine/Programs/UnrealGameSync/UnrealGameSync.ini
--- a/Engine/Programs/UnrealGameSync/UnrealGameSync.ini
+++ b/Engine/Programs/UnrealGameSync/UnrealGameSync.ini
@@ -33,6 +33,8 @@ DebugGameHasSeparateExecutable=1
 +Step=(UniqueId="46312669-5069-428d-8d72-c2416c5a322e", Description="Launch UnrealInsights", StatusText="Running UnrealInsights", EstimatedDuration="1", Type="Other", FileName="$(BranchDir)\\Engine\\Binaries\\Win64\\UnrealInsights.exe", Arguments="", bUseLogWindow="False", bNormalSync="False", bScheduledSync="False", bShowAsTool="True", Requires="89fe8a79-d259-4c7b-bfb4-68f7218b91c2", OrderIndex=21 )
 +Step=(UniqueId="bb48ca5b-5682-4432-824d-c451336a6523", Description="Compile Zen Dashboard", StatusText="Compile ZenDashboard Step...", EstimatedDuration="1", Type="Compile", Target="ZenDashboard", Platform="Win64", Configuration="Development", Arguments="", UseLogWindow="True", bNormalSync="False", bScheduledSync="False", bShowAsTool="True")
 +Step=(UniqueId="586cc0d3-3914-4df9-acb6-2c02cd9d4fc6", Description="Launch Zen Dashboard", StatusText="Running Zen Dashboard", EstimatedDuration="1", Type="Other", FileName="$(BranchDir)\\Engine\\Binaries\\Win64\\ZenDashboard.exe", WorkingDir="", Arguments="", bUseLogWindow="False", bNormalSync="False", bShowAsTool="False", bScheduledSync="False", Link="Zen", Requires="bb48ca5b-5682-4432-824d-c451336a6523" )
++Step=(UniqueId="91c2a429-c391-49b4-92a5-4e6be71e0f00", Description="Compile SwitchboardListener", StatusText="Compiling SwitchboardListener...", EstimatedDuration="1", Type="Compile", Target="SwitchboardListener", Platform="Win64", Configuration="Development", Arguments="", bNormalSync="False", bScheduledSync="False", bShowAsTool="False")
++Step=(UniqueId="5036c75b-8df0-4329-82a1-869dd2d48605", Description="Compile UnrealMultiUserServer", StatusText="Compiling UnrealMultiUserServer...", EstimatedDuration="1", Type="Compile", Target="UnrealMultiUserServer", Platform="Win64", Configuration="Development", Arguments="", bNormalSync="False", bScheduledSync="False", bShowAsTool="False")
 
 [Default]
 EditorTarget=UnrealEditor
