Bump Linux to new v20 toolchain

diff --git a/Engine/Programs/UnrealGameSync/UnrealGameSync.ini b/Engine/Programs/UnrealGameSync/UnrealGameSync.ini
--- a/Engine/Programs/UnrealGameSync/UnrealGameSync.ini
+++ b/Engine/Programs/UnrealGameSync/UnrealGameSync.ini
@@ -48,7 +48,7 @@ SdkInstallerDir=
 +SdkInfo=(Category="Windows", Description="Visual Studio 2019 v16.11.5")
 +SdkInfo=(Category="Windows", Description="Windows SDK 10.0.18362.0")
 +SdkInfo=(Category="Windows", Description="LLVM Clang 11.1.0")
-+SdkInfo=(Category="Linux", Description="v19_clang-11.0.1-centos7", Browse="$(SdkInstallerDir)\\Linux", Install="$(SdkInstallerDir)\\Linux\\v19_clang-11.0.1-centos7.exe")
++SdkInfo=(Category="Linux", Description="v20_clang-13.0.1-centos7", Browse="$(SdkInstallerDir)\\Linux", Install="$(SdkInstallerDir)\\Linux\\v20_clang-13.0.1-centos7.exe")
 +SdkInfo=(Category="Mac", Description="macOS 10.15.7")
 +SdkInfo=(Category="Mac", Description="Xcode 12.4")
 
