Updated UGS SDK info for Android to NDK r25b to match 5.1 requirements

diff --git a/Engine/Programs/UnrealGameSync/UnrealGameSync.ini b/Engine/Programs/UnrealGameSync/UnrealGameSync.ini
--- a/Engine/Programs/UnrealGameSync/UnrealGameSync.ini
+++ b/Engine/Programs/UnrealGameSync/UnrealGameSync.ini
@@ -44,7 +44,7 @@ EditorTarget=UnrealEditor
 SdkInstallerDir= 
 
 ; All the required SDKs for the current version of the engine
-+SdkInfo=(Category="Android", Description="NDK r21", Browse="$(SdkInstallerDir)\\Android")
++SdkInfo=(Category="Android", Description="NDK r25b", Browse="$(SdkInstallerDir)\\Android")
 +SdkInfo=(Category="Windows", Description="Visual Studio 2019 v16.11.5")
 +SdkInfo=(Category="Windows", Description="Windows SDK 10.0.18362.0")
 +SdkInfo=(Category="Windows", Description="LLVM Clang 11.1.0")
