Add InterchangeWorker to the UGS build list

#rnx

diff --git a/Engine/Programs/UnrealGameSync/UnrealGameSync.ini b/Engine/Programs/UnrealGameSync/UnrealGameSync.ini
--- a/Engine/Programs/UnrealGameSync/UnrealGameSync.ini
+++ b/Engine/Programs/UnrealGameSync/UnrealGameSync.ini
@@ -37,6 +37,7 @@ BuildUnrealHeaderTool=0
 +Step=(UniqueId="91c2a429-c391-49b4-92a5-4e6be71e0f00", Description="Compile SwitchboardListener", StatusText="Compiling SwitchboardListener...", EstimatedDuration="1", Type="Compile", Target="SwitchboardListener", Platform="Win64", Configuration="Development", Arguments="", bNormalSync="False", bScheduledSync="False", bShowAsTool="False")
 +Step=(UniqueId="5036c75b-8df0-4329-82a1-869dd2d48605", Description="Compile UnrealMultiUserServer", StatusText="Compiling UnrealMultiUserServer...", EstimatedDuration="1", Type="Compile", Target="UnrealMultiUserServer", Platform="Win64", Configuration="Development", Arguments="", bNormalSync="False", bScheduledSync="False", bShowAsTool="False")
 +Step=(UniqueId="274b89c3-9dc6-4465-a508-40abc4593cc2", Description="Compile UnrealMultiUserSlateServer", StatusText="Compiling UnrealMultiUserSlateServer...", EstimatedDuration="1", Type="Compile", Target="UnrealMultiUserSlateServer", Platform="Win64", Configuration="Development", Arguments="", bNormalSync="False", bScheduledSync="False", bShowAsTool="False")
++Step=(UniqueId="f7f0c4c7-5551-4485-a961-f3e4fd9d138c", Description="Compile InterchangeWorker", StatusText="Compiling InterchangeWorker...", EstimatedDuration="1", Type="Compile", Target="InterchangeWorker", Platform="Win64", Configuration="Development", Arguments="", bNormalSync="False", bScheduledSync="False", bShowAsTool="False")
 
 [Default]
 EditorTarget=UnrealEditor
