UGS: Added missing entry for SwitchboardListenerHelper related to CL 24906937

diff --git a/Engine/Programs/UnrealGameSync/UnrealGameSync.ini b/Engine/Programs/UnrealGameSync/UnrealGameSync.ini
--- a/Engine/Programs/UnrealGameSync/UnrealGameSync.ini
+++ b/Engine/Programs/UnrealGameSync/UnrealGameSync.ini
@@ -35,6 +35,7 @@ BuildUnrealHeaderTool=0
 +Step=(UniqueId="bb48ca5b-5682-4432-824d-c451336a6523", Description="Compile Zen Dashboard", StatusText="Compile ZenDashboard Step...", EstimatedDuration="1", Type="Compile", Target="ZenDashboard", Platform="Win64", Configuration="Development", Arguments="", UseLogWindow="True", bNormalSync="False", bScheduledSync="False", bShowAsTool="True")
 +Step=(UniqueId="586cc0d3-3914-4df9-acb6-2c02cd9d4fc6", Description="Launch Zen Dashboard", StatusText="Running Zen Dashboard", EstimatedDuration="1", Type="Other", FileName="$(BranchDir)\\Engine\\Binaries\\Win64\\ZenDashboard.exe", WorkingDir="", Arguments="", bUseLogWindow="False", bNormalSync="False", bShowAsTool="False", bScheduledSync="False", Link="Zen", Requires="bb48ca5b-5682-4432-824d-c451336a6523" )
 +Step=(UniqueId="91c2a429-c391-49b4-92a5-4e6be71e0f00", Description="Compile SwitchboardListener", StatusText="Compiling SwitchboardListener...", EstimatedDuration="1", Type="Compile", Target="SwitchboardListener", Platform="Win64", Configuration="Development", Arguments="", bNormalSync="False", bScheduledSync="False", bShowAsTool="False")
++Step=(UniqueId="71e65719-b7bb-49c5-961b-620acf1e1dad", Description="Compile SwitchboardListenerHelper", StatusText="Compiling SwitchboardListenerHelper...", EstimatedDuration="1", Type="Compile", Target="SwitchboardListenerHelper", Platform="Win64", Configuration="Development", Arguments="", bNormalSync="False", bScheduledSync="False", bShowAsTool="False")
 +Step=(UniqueId="5036c75b-8df0-4329-82a1-869dd2d48605", Description="Compile UnrealMultiUserServer", StatusText="Compiling UnrealMultiUserServer...", EstimatedDuration="1", Type="Compile", Target="UnrealMultiUserServer", Platform="Win64", Configuration="Development", Arguments="", bNormalSync="False", bScheduledSync="False", bShowAsTool="False")
 +Step=(UniqueId="274b89c3-9dc6-4465-a508-40abc4593cc2", Description="Compile UnrealMultiUserSlateServer", StatusText="Compiling UnrealMultiUserSlateServer...", EstimatedDuration="1", Type="Compile", Target="UnrealMultiUserSlateServer", Platform="Win64", Configuration="Development", Arguments="", bNormalSync="False", bScheduledSync="False", bShowAsTool="False")
 +Step=(UniqueId="f7f0c4c7-5551-4485-a961-f3e4fd9d138c", Description="Compile InterchangeWorker", StatusText="Compiling InterchangeWorker...", EstimatedDuration="1", Type="Compile", Target="InterchangeWorker", Platform="Win64", Configuration="Development", Arguments="", bNormalSync="False", bScheduledSync="False", bShowAsTool="False")
