Turn off new Linux toolchain until I can resolve the lld issues we're seeing on windows

diff --git a/Engine/Config/Linux/Linux_SDK.json b/Engine/Config/Linux/Linux_SDK.json
--- a/Engine/Config/Linux/Linux_SDK.json
+++ b/Engine/Config/Linux/Linux_SDK.json
@@ -1,7 +1,6 @@
-{
-"MainVersion" : "v23_clang-18.1.0-rockylinux8",
-"MinVersion" : "v23_clang-18.1.0-rockylinux8",
-"MaxVersion" : "v23_clang-18.1.0-rockylinux8",
+"MainVersion" : "v22_clang-16.0.6-centos7",
+"MinVersion" : "v22_clang-16.0.6-centos7",
+"MaxVersion" : "v22_clang-16.0.6-centos7",
 
 "AutoSDKPlatform" : "Linux_x64",
 }
diff --git a/Engine/Programs/UnrealGameSync/UnrealGameSync.ini b/Engine/Programs/UnrealGameSync/UnrealGameSync.ini
--- a/Engine/Programs/UnrealGameSync/UnrealGameSync.ini
+++ b/Engine/Programs/UnrealGameSync/UnrealGameSync.ini
@@ -55,6 +55,6 @@ SdkInstallerDir=
 +SdkInfo=(Category="Windows", Description="Visual Studio 2022 v18.4", Browse="https://visualstudio.microsoft.com/downloads/")
 +SdkInfo=(Category="Windows", Description="Windows SDK 10.0.18362.0", Browse="https://developer.microsoft.com/en-us/windows/downloads/windows-sdk/")
 +SdkInfo=(Category="Windows", Description="LLVM Clang 16.0.6", Browse="https://releases.llvm.org/")
-+SdkInfo=(Category="Linux", Description="v23_clang-18.1.0-rockylinux8", Install="$(SdkInstallerDir)\\Linux\\HostWin64\\v23_clang-18.1.0-rockylinux8.exe", Browse="$(SdkInstallerDir)\\Linux\\HostWin64")
++SdkInfo=(Category="Linux", Description="v22_clang-16.0.6-centos7", Install="$(SdkInstallerDir)\\Linux\\HostWin64\\v22_clang-16.0.6-centos7.exe", Browse="$(SdkInstallerDir)\\Linux\\HostWin64")
 +SdkInfo=(Category="Mac", Description="macOS 10.15.7")
 +SdkInfo=(Category="Mac", Description="Xcode 12.4")
