Update preferred Clang version to 19

diff --git a/Engine/Config/Windows/Windows_SDK.json b/Engine/Config/Windows/Windows_SDK.json
--- a/Engine/Config/Windows/Windows_SDK.json
+++ b/Engine/Config/Windows/Windows_SDK.json
@@ -6,8 +6,8 @@
 	"MinSoftwareVersion": "10.0.19041.0",
 
 	"PreferredClangVersions": [
-		"18.1.3-18.999",
-		"19.1.0-19.999"
+		"19.1.0-19.999",
+		"18.1.3-18.999"
 	],
 	"MinimumClangVersion": "18.1.3",
 
diff --git a/Engine/Programs/UnrealGameSync/UnrealGameSync.ini b/Engine/Programs/UnrealGameSync/UnrealGameSync.ini
--- a/Engine/Programs/UnrealGameSync/UnrealGameSync.ini
+++ b/Engine/Programs/UnrealGameSync/UnrealGameSync.ini
@@ -53,7 +53,7 @@ SdkInstallerDir=
 +SdkInfo=(Category="Android", Description="NDK r25b", Browse="$(SdkInstallerDir)\\Android")
 +SdkInfo=(Category="Windows", Description="Visual Studio 2022 17.10", Browse="https://visualstudio.microsoft.com/downloads/")
 +SdkInfo=(Category="Windows", Description="Windows SDK 10.0.18362.0", Browse="https://developer.microsoft.com/en-us/windows/downloads/windows-sdk/")
-+SdkInfo=(Category="Windows", Description="LLVM Clang 18.1.8", Browse="https://releases.llvm.org/")
++SdkInfo=(Category="Windows", Description="LLVM Clang 19.1.7", Browse="https://releases.llvm.org/")
 +SdkInfo=(Category="Linux", Description="v24_clang-19.1.7-rockylinux8", Install="$(SdkInstallerDir)\\Linux\\HostWin64\\v24_clang-19.1.7-rockylinux8.exe", Browse="$(SdkInstallerDir)\\Linux\\HostWin64")
 +SdkInfo=(Category="Mac", Description="macOS 10.15.7")
 +SdkInfo=(Category="Mac", Description="Xcode 12.4")
