[Linux] Fix up UnrealGameSync.ini to point to the proper Linux toolchain to be used on Windows

diff --git a/Engine/Programs/UnrealGameSync/UnrealGameSync.ini b/Engine/Programs/UnrealGameSync/UnrealGameSync.ini
--- a/Engine/Programs/UnrealGameSync/UnrealGameSync.ini
+++ b/Engine/Programs/UnrealGameSync/UnrealGameSync.ini
@@ -54,6 +54,6 @@ SdkInstallerDir=
 +SdkInfo=(Category="Windows", Description="Visual Studio 2022 17.10", Browse="https://visualstudio.microsoft.com/downloads/")
 +SdkInfo=(Category="Windows", Description="Windows SDK 10.0.18362.0", Browse="https://developer.microsoft.com/en-us/windows/downloads/windows-sdk/")
 +SdkInfo=(Category="Windows", Description="LLVM Clang 19.1.7", Browse="https://releases.llvm.org/")
-+SdkInfo=(Category="Linux", Description="v24_clang-19.1.7-rockylinux8", Install="$(SdkInstallerDir)\\Linux\\HostWin64\\v24_clang-19.1.7-rockylinux8.exe", Browse="$(SdkInstallerDir)\\Linux\\HostWin64")
++SdkInfo=(Category="Linux", Description="v25_clang-18.1.0-rockylinux8", Install="$(SdkInstallerDir)\\Linux\\HostWin64\\v25_clang-18.1.0-rockylinux8.exe", Browse="$(SdkInstallerDir)\\Linux\\HostWin64")
 +SdkInfo=(Category="Mac", Description="macOS 10.15.7")
 +SdkInfo=(Category="Mac", Description="Xcode 12.4")
