[Linux] Enable llvm 20.1.8

diff --git a/Engine/Config/Linux/Linux_SDK.json b/Engine/Config/Linux/Linux_SDK.json
--- a/Engine/Config/Linux/Linux_SDK.json
+++ b/Engine/Config/Linux/Linux_SDK.json
@@ -1,7 +1,7 @@
 {
-	"MainVersion" : "v25_clang-18.1.0-rockylinux8",
-	"MinVersion" : "v25_clang-18.1.0-rockylinux8",
-	"MaxVersion" : "v25_clang-18.1.0-rockylinux8",
+	"MainVersion" : "v26_clang-20.1.8-rockylinux8",
+	"MinVersion" : "v26_clang-20.1.8-rockylinux8",
+	"MaxVersion" : "v26_clang-20.1.8-rockylinux8",
 
 	"AutoSDKPlatform" : "Linux_x64"
 }
diff --git a/Engine/Programs/UnrealGameSync/UnrealGameSync.ini b/Engine/Programs/UnrealGameSync/UnrealGameSync.ini
--- a/Engine/Programs/UnrealGameSync/UnrealGameSync.ini
+++ b/Engine/Programs/UnrealGameSync/UnrealGameSync.ini
@@ -54,6 +54,6 @@ SdkInstallerDir=
 +SdkInfo=(Category="Windows", Description="Visual Studio 2022 17.14", Browse="https://visualstudio.microsoft.com/downloads/")
 +SdkInfo=(Category="Windows", Description="Windows SDK 10.0.18362.0", Browse="https://developer.microsoft.com/en-us/windows/downloads/windows-sdk/")
 +SdkInfo=(Category="Windows", Description="LLVM Clang 19.1.7", Browse="https://releases.llvm.org/")
-+SdkInfo=(Category="Linux", Description="v25_clang-18.1.0-rockylinux8", Install="$(SdkInstallerDir)\\Linux\\HostWin64\\v25_clang-18.1.0-rockylinux8.exe", Browse="$(SdkInstallerDir)\\Linux\\HostWin64")
++SdkInfo=(Category="Linux", Description="v26_clang-20.1.8-rockylinux8", Install="$(SdkInstallerDir)\\Linux\\HostWin64\\v26_clang-20.1.8-rockylinux8.exe", Browse="$(SdkInstallerDir)\\Linux\\HostWin64")
 +SdkInfo=(Category="Mac", Description="macOS 10.15.7")
 +SdkInfo=(Category="Mac", Description="Xcode 12.4")
diff --git a/Engine/Source/Programs/UnrealBuildTool/Platform/Linux/LinuxToolChain.cs b/Engine/Source/Programs/UnrealBuildTool/Platform/Linux/LinuxToolChain.cs
--- a/Engine/Source/Programs/UnrealBuildTool/Platform/Linux/LinuxToolChain.cs
+++ b/Engine/Source/Programs/UnrealBuildTool/Platform/Linux/LinuxToolChain.cs
@@ -55,7 +55,7 @@ namespace UnrealBuildTool
 			: this(UnrealTargetPlatform.Linux, InArchitecture, InSDK, InOptions, InLogger)
 		{
 			// prevent unknown clangs since the build is likely to fail on too old or too new compilers
-			if ((CompilerVersionLessThan(18, 0, 0) || CompilerVersionGreaterOrEqual(19, 0, 0)) && !Options.HasFlag(ClangToolChainOptions.UseAutoRTFMCompiler))
+			if ((CompilerVersionLessThan(20, 0, 0) || CompilerVersionGreaterOrEqual(21, 0, 0)) && !Options.HasFlag(ClangToolChainOptions.UseAutoRTFMCompiler))
 			{
 				throw new BuildLogEventException("This version of the Unreal Engine can only be compiled with clang 18.x. clang {Version} may not build it - please use a different version.", Info.ClangVersion);
 			}
