Add UnrealFrontend to the UnrealGameSync tools menu

#rnx

#Synced-CL 50394790

diff --git a/Engine/Programs/UnrealGameSync/UnrealGameSync.ini b/Engine/Programs/UnrealGameSync/UnrealGameSync.ini
--- a/Engine/Programs/UnrealGameSync/UnrealGameSync.ini
+++ b/Engine/Programs/UnrealGameSync/UnrealGameSync.ini
@@ -40,6 +40,8 @@ BuildUnrealHeaderTool=0
 +Step=(UniqueId="fffcbd4b-0f14-40ca-9069-f09c1c6e5898", Description="Compile LiveLinkHub", StatusText="Compiling LiveLinkHub...", EstimatedDuration="1", Type="Compile", Target="LiveLinkHub", Platform="Win64", Configuration="Development", Arguments="", bNormalSync="False", bScheduledSync="False", bShowAsTool="False")
 +Step=(UniqueId="0aadd9c7-17b7-4818-a9ac-d11dc71f6ab9", Description="Launch LiveLinkHub", StatusText="Running LiveLinkHub", EstimatedDuration="1", Type="Other", FileName="$(BranchDir)\\Engine\\Binaries\\Win64\\LiveLinkHub.exe", WorkingDir="", Arguments="", bUseLogWindow="False", bNormalSync="False", bShowAsTool="True", bScheduledSync="False", Requires="fffcbd4b-0f14-40ca-9069-f09c1c6e5898" )
 +Step=(UniqueId="f7f0c4c7-5551-4485-a961-f3e4fd9d138c", Description="Compile InterchangeWorker", StatusText="Compiling InterchangeWorker...", EstimatedDuration="1", Type="Compile", Target="InterchangeWorker", Platform="Win64", Configuration="Development", Arguments="", bNormalSync="False", bScheduledSync="False", bShowAsTool="False")
++Step=(UniqueId="810b581a-35fc-4550-9f1f-234494ae7c84", Description="Compile UnrealFrontend", StatusText="Compile UnrealFrontend Step...", EstimatedDuration="1", Type="Compile", Target="UnrealFrontend", Platform="Win64", Configuration="Development", Arguments="", bNormalSync="False", bScheduledSync="False", bShowAsTool="False", OrderIndex=22)
++Step=(UniqueId="8536e8df-9116-41e6-8f84-20018a31221c", Description="Launch UnrealFrontend", StatusText="Running UnrealFrontend", EstimatedDuration="1", Type="Other", FileName="$(BranchDir)\\Engine\\Binaries\\Win64\\UnrealFrontend.exe", Arguments="", bUseLogWindow="False", bNormalSync="False", bScheduledSync="False", bShowAsTool="True", Requires="810b581a-35fc-4550-9f1f-234494ae7c84", OrderIndex=23 )
 
 [Default]
 EditorTarget=UnrealEditor
